Fixed Copywrite.

diff --git a/Engine/Source/Editor/EditorSubsystem/EditorSubsystem.Build.cs b/Engine/Source/Editor/EditorSubsystem/EditorSubsystem.Build.cs
--- a/Engine/Source/Editor/EditorSubsystem/EditorSubsystem.Build.cs
+++ b/Engine/Source/Editor/EditorSubsystem/EditorSubsystem.Build.cs
@@ -1,4 +1,4 @@
-// Copyright 1998-2018 Epic Games, Inc. All Rights Reserved.
+// Copyright 1998-2019 Epic Games, Inc. All Rights Reserved.
 
 using UnrealBuildTool;
 
