Updated a number of copywrite lines in new files.

diff --git a/Engine/Source/Editor/EditorSubsystem/Private/EditorSubsystem.cpp b/Engine/Source/Editor/EditorSubsystem/Private/EditorSubsystem.cpp
--- a/Engine/Source/Editor/EditorSubsystem/Private/EditorSubsystem.cpp
+++ b/Engine/Source/Editor/EditorSubsystem/Private/EditorSubsystem.cpp
@@ -1,4 +1,4 @@
-// Copyright 1998-2018 Epic Games, Inc. All Rights Reserved.
+// Copyright 1998-2019 Epic Games, Inc. All Rights Reserved.
 
 #include "EditorSubsystem.h"
 
diff --git a/Engine/Source/Editor/EditorSubsystem/Private/EditorSubsystemModule.cpp b/Engine/Source/Editor/EditorSubsystem/Private/EditorSubsystemModule.cpp
--- a/Engine/Source/Editor/EditorSubsystem/Private/EditorSubsystemModule.cpp
+++ b/Engine/Source/Editor/EditorSubsystem/Private/EditorSubsystemModule.cpp
@@ -1,4 +1,4 @@
-// Copyright 1998-2018 Epic Games, Inc. All Rights Reserved.
+// Copyright 1998-2019 Epic Games, Inc. All Rights Reserved.
 
 #include "EditorSubsystemModule.h"
 #include "Modules/ModuleManager.h"
diff --git a/Engine/Source/Editor/EditorSubsystem/Public/EditorSubsystem.h b/Engine/Source/Editor/EditorSubsystem/Public/EditorSubsystem.h
--- a/Engine/Source/Editor/EditorSubsystem/Public/EditorSubsystem.h
+++ b/Engine/Source/Editor/EditorSubsystem/Public/EditorSubsystem.h
@@ -1,4 +1,4 @@
-// Copyright 1998-2018 Epic Games, Inc. All Rights Reserved.
+// Copyright 1998-2019 Epic Games, Inc. All Rights Reserved.
 
 #pragma once
 
diff --git a/Engine/Source/Editor/EditorSubsystem/Public/EditorSubsystemModule.h b/Engine/Source/Editor/EditorSubsystem/Public/EditorSubsystemModule.h
--- a/Engine/Source/Editor/EditorSubsystem/Public/EditorSubsystemModule.h
+++ b/Engine/Source/Editor/EditorSubsystem/Public/EditorSubsystemModule.h
@@ -1,4 +1,4 @@
-// Copyright 1998-2018 Epic Games, Inc. All Rights Reserved.
+// Copyright 1998-2019 Epic Games, Inc. All Rights Reserved.
 
 #pragma once
 
diff --git a/Engine/Source/Editor/UnrealEd/Private/Subsystems/ImportSubsystem.cpp b/Engine/Source/Editor/UnrealEd/Private/Subsystems/ImportSubsystem.cpp
--- a/Engine/Source/Editor/UnrealEd/Private/Subsystems/ImportSubsystem.cpp
+++ b/Engine/Source/Editor/UnrealEd/Private/Subsystems/ImportSubsystem.cpp
@@ -1,4 +1,4 @@
-// Copyright 1998-2018 Epic Games, Inc. All Rights Reserved.
+// Copyright 1998-2019 Epic Games, Inc. All Rights Reserved.
 
 #include "Subsystems/ImportSubsystem.h"
 
diff --git a/Engine/Source/Editor/UnrealEd/Public/Subsystems/ImportSubsystem.h b/Engine/Source/Editor/UnrealEd/Public/Subsystems/ImportSubsystem.h
--- a/Engine/Source/Editor/UnrealEd/Public/Subsystems/ImportSubsystem.h
+++ b/Engine/Source/Editor/UnrealEd/Public/Subsystems/ImportSubsystem.h
@@ -1,4 +1,4 @@
-// Copyright 1998-2018 Epic Games, Inc. All Rights Reserved.
+// Copyright 1998-2019 Epic Games, Inc. All Rights Reserved.
 
 #pragma once
 
diff --git a/Engine/Source/Runtime/Engine/Classes/Engine/ContentEncryptionConfig.h b/Engine/Source/Runtime/Engine/Classes/Engine/ContentEncryptionConfig.h
--- a/Engine/Source/Runtime/Engine/Classes/Engine/ContentEncryptionConfig.h
+++ b/Engine/Source/Runtime/Engine/Classes/Engine/ContentEncryptionConfig.h
@@ -1,4 +1,4 @@
-// Copyright 1998-2018 Epic Games, Inc. All Rights Reserved.
+// Copyright 1998-2019 Epic Games, Inc. All Rights Reserved.
 #pragma once
 
 /** Project specific configuration for content encryption */
diff --git a/Engine/Source/Runtime/Engine/Private/Subsystems/EngineSubsystem.cpp b/Engine/Source/Runtime/Engine/Private/Subsystems/EngineSubsystem.cpp
--- a/Engine/Source/Runtime/Engine/Private/Subsystems/EngineSubsystem.cpp
+++ b/Engine/Source/Runtime/Engine/Private/Subsystems/EngineSubsystem.cpp
@@ -1,4 +1,4 @@
-// Copyright 1998-2018 Epic Games, Inc. All Rights Reserved.
+// Copyright 1998-2019 Epic Games, Inc. All Rights Reserved.
 
 #include "Subsystems/EngineSubsystem.h"
 
diff --git a/Engine/Source/Runtime/Engine/Public/Subsystems/EngineSubsystem.h b/Engine/Source/Runtime/Engine/Public/Subsystems/EngineSubsystem.h
--- a/Engine/Source/Runtime/Engine/Public/Subsystems/EngineSubsystem.h
+++ b/Engine/Source/Runtime/Engine/Public/Subsystems/EngineSubsystem.h
@@ -1,4 +1,4 @@
-// Copyright 1998-2018 Epic Games, Inc. All Rights Reserved.
+// Copyright 1998-2019 Epic Games, Inc. All Rights Reserved.
 
 #pragma once
 
